Missing copyright notices.

#rnx

diff --git a/Engine/Source/Programs/UnrealBuildTool/Storage/IStorageProvider.cs b/Engine/Source/Programs/UnrealBuildTool/Storage/IStorageProvider.cs
--- a/Engine/Source/Programs/UnrealBuildTool/Storage/IStorageProvider.cs
+++ b/Engine/Source/Programs/UnrealBuildTool/Storage/IStorageProvider.cs
@@ -1,3 +1,5 @@
+// Copyright Epic Games, Inc. All Rights Reserved.
+
 using System;
 using System.Collections.Generic;
 using System.Diagnostics;
diff --git a/Engine/Source/Programs/UnrealBuildTool/Storage/Impl/FileSystemStorageProvider.cs b/Engine/Source/Programs/UnrealBuildTool/Storage/Impl/FileSystemStorageProvider.cs
--- a/Engine/Source/Programs/UnrealBuildTool/Storage/Impl/FileSystemStorageProvider.cs
+++ b/Engine/Source/Programs/UnrealBuildTool/Storage/Impl/FileSystemStorageProvider.cs
@@ -1,3 +1,5 @@
+// Copyright Epic Games, Inc. All Rights Reserved.
+
 using System;
 using System.Collections.Generic;
 using System.Diagnostics;
diff --git a/Engine/Source/Programs/UnrealBuildToolTests/PreprocessorStateTests.cs b/Engine/Source/Programs/UnrealBuildToolTests/PreprocessorStateTests.cs
--- a/Engine/Source/Programs/UnrealBuildToolTests/PreprocessorStateTests.cs
+++ b/Engine/Source/Programs/UnrealBuildToolTests/PreprocessorStateTests.cs
@@ -1,4 +1,4 @@
-// Copyright 1998-2019 Epic Games, Inc. All Rights Reserved.
+// Copyright Epic Games, Inc. All Rights Reserved.
 
 using Microsoft.VisualStudio.TestTools.UnitTesting;
 using System;
diff --git a/Engine/Source/Programs/UnrealBuildToolTests/PreprocessorTests.cs b/Engine/Source/Programs/UnrealBuildToolTests/PreprocessorTests.cs
--- a/Engine/Source/Programs/UnrealBuildToolTests/PreprocessorTests.cs
+++ b/Engine/Source/Programs/UnrealBuildToolTests/PreprocessorTests.cs
@@ -1,3 +1,5 @@
+// Copyright Epic Games, Inc. All Rights Reserved.
+
 using Microsoft.VisualStudio.TestTools.UnitTesting;
 using System.Collections.Generic;
 using System.Text;
diff --git a/Engine/Source/Programs/UnrealBuildToolTests/SourceFileTests.cs b/Engine/Source/Programs/UnrealBuildToolTests/SourceFileTests.cs
--- a/Engine/Source/Programs/UnrealBuildToolTests/SourceFileTests.cs
+++ b/Engine/Source/Programs/UnrealBuildToolTests/SourceFileTests.cs
@@ -1,4 +1,4 @@
-// Copyright 1998-2019 Epic Games, Inc. All Rights Reserved.
+// Copyright Epic Games, Inc. All Rights Reserved.
 
 using Microsoft.VisualStudio.TestTools.UnitTesting;
 using System;
diff --git a/Engine/Source/Programs/UnrealBuildToolTests/TokenReaderTests.cs b/Engine/Source/Programs/UnrealBuildToolTests/TokenReaderTests.cs
--- a/Engine/Source/Programs/UnrealBuildToolTests/TokenReaderTests.cs
+++ b/Engine/Source/Programs/UnrealBuildToolTests/TokenReaderTests.cs
@@ -1,4 +1,4 @@
-// Copyright 1998-2019 Epic Games, Inc. All Rights Reserved.
+// Copyright Epic Games, Inc. All Rights Reserved.
 
 using Microsoft.VisualStudio.TestTools.UnitTesting;
 using System;
