Copyright fixups

diff --git a/Engine/Source/Programs/UnrealBuildTool/Platform/Windows/VCCompileAction.cs b/Engine/Source/Programs/UnrealBuildTool/Platform/Windows/VCCompileAction.cs
--- a/Engine/Source/Programs/UnrealBuildTool/Platform/Windows/VCCompileAction.cs
+++ b/Engine/Source/Programs/UnrealBuildTool/Platform/Windows/VCCompileAction.cs
@@ -1,3 +1,5 @@
+// Copyright Epic Games, Inc. All Rights Reserved.
+
 using System;
 using System.Collections.Generic;
 using System.Linq;
diff --git a/Engine/Source/Programs/UnrealBuildToolTests/PreprocessorExpressionTests.cs b/Engine/Source/Programs/UnrealBuildToolTests/PreprocessorExpressionTests.cs
--- a/Engine/Source/Programs/UnrealBuildToolTests/PreprocessorExpressionTests.cs
+++ b/Engine/Source/Programs/UnrealBuildToolTests/PreprocessorExpressionTests.cs
@@ -1,4 +1,4 @@
-// Copyright 1998-2019 Epic Games, Inc. All Rights Reserved.
+// Copyright Epic Games, Inc. All Rights Reserved.
 
 using Microsoft.VisualStudio.TestTools.UnitTesting;
 using System;
