Renamed UnrealBuildToolCore.sln to just UnrealBuildTool and updated with the new project names after we switched it over to net core.

diff --git a/Engine/Source/Programs/UnrealBuildToolTests/UnrealBuildToolTests.csproj b/Engine/Source/Programs/UnrealBuildToolTests/UnrealBuildToolTests.csproj
--- a/Engine/Source/Programs/UnrealBuildToolTests/UnrealBuildToolTests.csproj
+++ b/Engine/Source/Programs/UnrealBuildToolTests/UnrealBuildToolTests.csproj
@@ -14,7 +14,7 @@
   </ItemGroup>
 
   <ItemGroup>
-    <ProjectReference Include="..\UnrealBuildTool\UnrealBuildToolCore.csproj" />
+    <ProjectReference Include="..\UnrealBuildTool\UnrealBuildTool.csproj" />
   </ItemGroup>
 
 </Project>
