Fix UnrealBuildToolTests compile error

diff --git a/Engine/Source/Programs/UnrealBuildToolTests/PreprocessorTests.cs b/Engine/Source/Programs/UnrealBuildToolTests/PreprocessorTests.cs
--- a/Engine/Source/Programs/UnrealBuildToolTests/PreprocessorTests.cs
+++ b/Engine/Source/Programs/UnrealBuildToolTests/PreprocessorTests.cs
@@ -225,7 +225,7 @@ namespace UnrealBuildToolTests
 		{
 			string[] Lines = Fragment.Split('\n');
 
-			SourceFile File = new SourceFile(TokenReader.GetNullTerminatedByteArray(Fragment));
+			SourceFile File = new SourceFile(null, TokenReader.GetNullTerminatedByteArray(Fragment));
 
 			Preprocessor Instance = new Preprocessor();
 
