UnrealBuildToolTests:

Fix after refactorings of UnrealBuildTool/BuildUtilities.

#trivial

diff --git a/Engine/Source/Programs/UnrealBuildToolTests/SourceFileTests.cs b/Engine/Source/Programs/UnrealBuildToolTests/SourceFileTests.cs
--- a/Engine/Source/Programs/UnrealBuildToolTests/SourceFileTests.cs
+++ b/Engine/Source/Programs/UnrealBuildToolTests/SourceFileTests.cs
@@ -8,6 +8,7 @@ using System.Linq;
 using System.Text;
 using System.Threading.Tasks;
 using EpicGames.Core;
+using UnrealBuildBase;
 using UnrealBuildTool;
 
 namespace UnrealBuildToolTests
@@ -49,7 +50,7 @@ namespace UnrealBuildToolTests
 
 			Log.TraceLog("Read {0} source files", SourceFiles.Count);
 
-			FileReference TempDataFile = FileReference.Combine(UnrealBuildTool.UnrealBuildTool.EngineDirectory, "Intermediate", "Temp", "SourceFileTests.bin");
+			FileReference TempDataFile = FileReference.Combine(UnrealBuild.EngineDirectory, "Intermediate", "Temp", "SourceFileTests.bin");
 			DirectoryReference.CreateDirectory(TempDataFile.Directory);
 
 			using(Timeline.ScopeEvent("Writing source file data"))
