UnrealBuildToolTests:

Update variable location

#trivial

diff --git a/Engine/Source/Programs/UnrealBuildToolTests/SourceFileTests.cs b/Engine/Source/Programs/UnrealBuildToolTests/SourceFileTests.cs
--- a/Engine/Source/Programs/UnrealBuildToolTests/SourceFileTests.cs
+++ b/Engine/Source/Programs/UnrealBuildToolTests/SourceFileTests.cs
@@ -50,7 +50,7 @@ namespace UnrealBuildToolTests
 
 			Log.TraceLog("Read {0} source files", SourceFiles.Count);
 
-			FileReference TempDataFile = FileReference.Combine(UnrealBuild.EngineDirectory, "Intermediate", "Temp", "SourceFileTests.bin");
+			FileReference TempDataFile = FileReference.Combine(Unreal.EngineDirectory, "Intermediate", "Temp", "SourceFileTests.bin");
 			DirectoryReference.CreateDirectory(TempDataFile.Directory);
 
 			using(Timeline.ScopeEvent("Writing source file data"))
