UnrealBuildToolTests:

Allow overriding the assembly location, to fix computation of RootLocation when running tests.

diff --git a/Engine/Source/Programs/DotNETCommon/BuildUtilities/Unreal.cs b/Engine/Source/Programs/DotNETCommon/BuildUtilities/Unreal.cs
--- a/Engine/Source/Programs/DotNETCommon/BuildUtilities/Unreal.cs
+++ b/Engine/Source/Programs/DotNETCommon/BuildUtilities/Unreal.cs
@@ -13,6 +13,11 @@ namespace UnrealBuildBase
 	{
 		private static DirectoryReference FindRootDirectory()
 		{
+			if (LocationOverride.RootDirectory != null)
+			{
+				return LocationOverride.RootDirectory;
+			}
+
 			// This base library may be used - and so be launched - from more than one location (at time of writing, UnrealBuildTool and AutomationTool)
 			// Programs that use this assembly must be located under "Engine/Binaries/DotNET" and so we look for that sequence of directories in that path of the executing assembly
 			
@@ -109,5 +114,13 @@ namespace UnrealBuildBase
 			}
 			return bIsEngineInstalled.Value;
 		}
+
+		public static class LocationOverride 
+		{
+			/// <summary>
+			/// If set, this value will be used to populate Unreal.RootDirectory
+			/// </summary>
+			public static DirectoryReference? RootDirectory = null;
+		}
 	}
 }
diff --git a/Engine/Source/Programs/UnrealBuildToolTests/SourceFileTests.cs b/Engine/Source/Programs/UnrealBuildToolTests/SourceFileTests.cs
--- a/Engine/Source/Programs/UnrealBuildToolTests/SourceFileTests.cs
+++ b/Engine/Source/Programs/UnrealBuildToolTests/SourceFileTests.cs
@@ -22,6 +22,11 @@ namespace UnrealBuildToolTests
 		[TestMethod]
 		public void Run()
 		{
+			// Test assemblies are built into a temporary directory outside of the unreal source tree, which prevents Unreal.FindRootDirectory()
+			// from being able to deduce the correct location - so override it.q
+			// The default working directory for this test is ...\Engine\Source\Programs\UnrealBuildToolTests\bin\[configuration]\netcoreapp3.1
+			Unreal.LocationOverride.RootDirectory = DirectoryReference.Combine(DirectoryReference.GetCurrentDirectory(), @"..\..\..\..\..\..\..");
+
 			List<DirectoryReference> BaseDirectories = new List<DirectoryReference>();
 			BaseDirectories.Add(DirectoryReference.Combine(UnrealBuildTool.UnrealBuildTool.EngineSourceDirectory, "Runtime"));
 			BaseDirectories.Add(DirectoryReference.Combine(UnrealBuildTool.UnrealBuildTool.EngineSourceDirectory, "Developer"));
