UnrealBuildToolTests:

Fix for namespace change made to BuildUtilities

#trivial

diff --git a/Engine/Source/Programs/UnrealBuildToolTests/SourceFileTests.cs b/Engine/Source/Programs/UnrealBuildToolTests/SourceFileTests.cs
--- a/Engine/Source/Programs/UnrealBuildToolTests/SourceFileTests.cs
+++ b/Engine/Source/Programs/UnrealBuildToolTests/SourceFileTests.cs
@@ -32,7 +32,7 @@ namespace UnrealBuildToolTests
 			BaseDirectories.Add(DirectoryReference.Combine(UnrealBuildTool.UnrealBuildTool.EngineSourceDirectory, "Developer"));
 			BaseDirectories.Add(DirectoryReference.Combine(UnrealBuildTool.UnrealBuildTool.EngineSourceDirectory, "Editor"));
 
-			foreach(FileReference PluginFile in Plugins.EnumeratePlugins((FileReference)null))
+			foreach(FileReference PluginFile in PluginsBase.EnumeratePlugins((FileReference)null))
 			{
 				DirectoryReference PluginSourceDir = DirectoryReference.Combine(PluginFile.Directory, "Source");
 				if(DirectoryReference.Exists(PluginSourceDir))
